ToolsVersion fix for Mono.

diff --git a/Engine/Source/Programs/EnvVarsToXML/EnvVarsToXML/EnvVarsToXML.csproj b/Engine/Source/Programs/EnvVarsToXML/EnvVarsToXML/EnvVarsToXML.csproj
--- a/Engine/Source/Programs/EnvVarsToXML/EnvVarsToXML/EnvVarsToXML.csproj
+++ b/Engine/Source/Programs/EnvVarsToXML/EnvVarsToXML/EnvVarsToXML.csproj
@@ -1,5 +1,5 @@
 ﻿<?xml version="1.0" encoding="utf-8"?>
-<Project ToolsVersion="12.0" DefaultTargets="Build" xmlns="http://schemas.microsoft.com/developer/msbuild/2003">
+<Project ToolsVersion="4.0" DefaultTargets="Build" xmlns="http://schemas.microsoft.com/developer/msbuild/2003">
   <Import Project="$(MSBuildExtensionsPath)\$(MSBuildToolsVersion)\Microsoft.Common.props" Condition="Exists('$(MSBuildExtensionsPath)\$(MSBuildToolsVersion)\Microsoft.Common.props')" />
   <PropertyGroup>
     <Configuration Condition=" '$(Configuration)' == '' ">Debug</Configuration>
