Fix casing of System.Xml reference in EnvVarsToXML to fix Mac build.

diff --git a/Engine/Source/Programs/EnvVarsToXML/EnvVarsToXML/EnvVarsToXML.csproj b/Engine/Source/Programs/EnvVarsToXML/EnvVarsToXML/EnvVarsToXML.csproj
--- a/Engine/Source/Programs/EnvVarsToXML/EnvVarsToXML/EnvVarsToXML.csproj
+++ b/Engine/Source/Programs/EnvVarsToXML/EnvVarsToXML/EnvVarsToXML.csproj
@@ -35,7 +35,7 @@
     <Reference Include="System" />
     <Reference Include="System.Core" />
     <Reference Include="Microsoft.CSharp" />
-    <Reference Include="System.XML" />
+    <Reference Include="System.Xml" />
   </ItemGroup>
   <ItemGroup>
     <Compile Include="Program.cs" />
